Reorganize StorageServerClient module into the base module and a sibling module for debug functionality.  Also expose an API for creation of storage server platform file usage outside of the default platform hierarchy.

diff --git a/Engine/Source/Runtime/Launch/Launch.Build.cs b/Engine/Source/Runtime/Launch/Launch.Build.cs
--- a/Engine/Source/Runtime/Launch/Launch.Build.cs
+++ b/Engine/Source/Runtime/Launch/Launch.Build.cs
@@ -247,6 +247,7 @@ public class Launch : ModuleRules
 		if (Target.Configuration != UnrealTargetConfiguration.Shipping)
 		{
 			PrivateDependencyModuleNames.Add("StorageServerClient");
+			PrivateDependencyModuleNames.Add("StorageServerClientDebug");
 
 			if (Target.Type != TargetType.Program)
 			{
diff --git a/Engine/Source/Runtime/StorageServerClient/Internal/IStorageServerPlatformFile.h b/Engine/Source/Runtime/StorageServerClient/Internal/IStorageServerPlatformFile.h
new file mode 100644
--- /dev/null
+++ b/Engine/Source/Runtime/StorageServerClient/Internal/IStorageServerPlatformFile.h
@@ -0,0 +1,27 @@
+// Copyright Epic Games, Inc. All Rights Reserved.
+
+#pragma once
+
+#include "Containers/StringFwd.h"
+#include "HAL/PlatformFile.h"
+
+#if !UE_BUILD_SHIPPING
+
+class IStorageServerPlatformFile : public IPlatformFile
+{
+public:
+
+	struct FConnectionStats
+	{
+		uint64 AccumulatedBytes = 0;
+		uint32 RequestCount = 0;
+		double MinRequestThroughput = 0.0;
+		double MaxRequestThroughput = 0.0;
+	};
+
+	virtual FStringView GetHostAddr() const = 0;
+
+	virtual void GetAndResetConnectionStats(FConnectionStats& OutStats) = 0;
+};
+
+#endif // !UE_BUILD_SHIPPING
diff --git a/Engine/Source/Runtime/StorageServerClient/Internal/StorageServerClientModule.h b/Engine/Source/Runtime/StorageServerClient/Internal/StorageServerClientModule.h
new file mode 100644
--- /dev/null
+++ b/Engine/Source/Runtime/StorageServerClient/Internal/StorageServerClientModule.h
@@ -0,0 +1,29 @@
+// Copyright Epic Games, Inc. All Rights Reserved.
+
+#pragma once
+
+#include "Containers/StringFwd.h"
+#include "HAL/IPlatformFileModule.h"
+#include "HAL/PlatformFileManager.h"
+#include "IStorageServerPlatformFile.h"
+#include "Modules/ModuleInterface.h"
+#include "Modules/ModuleManager.h"
+
+#if !UE_BUILD_SHIPPING
+
+class IStorageServerClientModule : public IPlatformFileModule
+{
+public:
+	static FORCEINLINE IStorageServerClientModule& Get()
+	{
+		return FModuleManager::LoadModuleChecked<IStorageServerClientModule>("StorageServerClient");
+	}
+	static FORCEINLINE IStorageServerPlatformFile* FindStorageServerPlatformFile()
+	{
+		return static_cast<IStorageServerPlatformFile*>(FPlatformFileManager::Get().FindPlatformFile(TEXT("StorageServer")));
+	}
+
+	virtual IStorageServerPlatformFile* TryCreateCustomPlatformFile(FStringView StoreDirectory, IPlatformFile* Inner) = 0;
+};
+
+#endif // !UE_BUILD_SHIPPING
diff --git a/Engine/Source/Runtime/StorageServerClient/Private/DebugStorageServerConnection.cpp b/Engine/Source/Runtime/StorageServerClient/Private/DebugStorageServerConnection.cpp
deleted file mode 100644
diff --git a/Engine/Source/Runtime/StorageServerClient/Private/DebugStorageServerConnection.h b/Engine/Source/Runtime/StorageServerClient/Private/DebugStorageServerConnection.h
deleted file mode 100644
diff --git a/Engine/Source/Runtime/StorageServerClient/Private/StorageServerClientModule.cpp b/Engine/Source/Runtime/StorageServerClient/Private/StorageServerClientModule.cpp
new file mode 100644
--- /dev/null
+++ b/Engine/Source/Runtime/StorageServerClient/Private/StorageServerClientModule.cpp
@@ -0,0 +1,75 @@
+// Copyright Epic Games, Inc. All Rights Reserved.
+
+#include "StorageServerClientModule.h"
+
+#include "Misc/CommandLine.h"
+#include "Misc/Paths.h"
+#include "StorageServerPlatformFile.h"
+
+#if !UE_BUILD_SHIPPING
+
+class FStorageServerClientModule : public IStorageServerClientModule
+{
+public:
+	FStorageServerClientModule() = default;
+	~FStorageServerClientModule() = default;
+
+	// ~Begin IModuleInterface
+	virtual void StartupModule() override;
+	virtual void ShutdownModule() override;
+	// ~End IModuleInterface
+
+	// ~Begin IPlatformFileModule
+	virtual IPlatformFile* GetPlatformFile() override;
+	// ~End IPlatformFileModule
+
+	// ~Begin IStorageServerClientModule
+	virtual IStorageServerPlatformFile* TryCreateCustomPlatformFile(FStringView StoreDirectory, IPlatformFile* Inner) override;
+	// ~End IStorageServerClientModule
+
+	IPlatformFile* CachedDefaultPlatformFileInstance = nullptr;
+};
+
+void FStorageServerClientModule::StartupModule()
+{
+}
+
+void FStorageServerClientModule::ShutdownModule()
+{
+}
+
+IPlatformFile* FStorageServerClientModule::GetPlatformFile()
+{
+	static TUniquePtr<IPlatformFile> DefaultPlatformFileInstance = MakeUnique<FStorageServerPlatformFile>();
+	if (!CachedDefaultPlatformFileInstance)
+	{
+		CachedDefaultPlatformFileInstance = DefaultPlatformFileInstance.Get();
+		// Load optional modules.
+		if (FModuleManager::Get().ModuleExists(TEXT("StorageServerClientDebug")))
+		{
+			FModuleManager::Get().LoadModule("StorageServerClientDebug");
+		}
+	}
+	return DefaultPlatformFileInstance.Get();
+}
+
+IStorageServerPlatformFile* FStorageServerClientModule::TryCreateCustomPlatformFile(FStringView StoreDirectory, IPlatformFile* Inner)
+{
+	TUniquePtr<FStorageServerPlatformFile> StorageServerPlatformFile = MakeUnique<FStorageServerPlatformFile>();
+	StorageServerPlatformFile->SetCustomProjectStorePath(StoreDirectory);
+	StorageServerPlatformFile->SetAllowPackageIo(false);
+	const TCHAR* CmdLine = FCommandLine::Get();
+	if (StorageServerPlatformFile->ShouldBeUsed(Inner, CmdLine) && StorageServerPlatformFile->Initialize(Inner, CmdLine))
+	{
+		if (FPaths::IsProjectFilePathSet())
+		{
+			StorageServerPlatformFile->InitializeAfterProjectFilePath();
+		}
+		return StorageServerPlatformFile.Release();
+	}
+	return nullptr;
+}
+
+IMPLEMENT_MODULE(FStorageServerClientModule, StorageServerClient)
+
+#endif // !UE_BUILD_SHIPPING
diff --git a/Engine/Source/Runtime/StorageServerClient/Private/StorageServerConnection.cpp b/Engine/Source/Runtime/StorageServerClient/Private/StorageServerConnection.cpp
--- a/Engine/Source/Runtime/StorageServerClient/Private/StorageServerConnection.cpp
+++ b/Engine/Source/Runtime/StorageServerClient/Private/StorageServerConnection.cpp
@@ -1,7 +1,6 @@
 // Copyright Epic Games, Inc. All Rights Reserved.
 
 #include "StorageServerConnection.h"
-#include "DebugStorageServerConnection.h"
 
 #include "IO/IoDispatcher.h"
 #include "IPAddress.h"
@@ -20,7 +19,7 @@
 #include "BuiltInHttpClient/BuiltInHttpClient.h"
 #include "BuiltInHttpClient/BuiltInHttpClientFSocket.h"
 #include "BuiltInHttpClient/BuiltInHttpClientPlatformSocket.h"
-#include "Engine/Engine.h"
+#include "HAL/PlatformMath.h"
 
 #if !UE_BUILD_SHIPPING
 
@@ -28,15 +27,6 @@ DEFINE_LOG_CATEGORY(LogStorageServerConnection);
 
 TRACE_DECLARE_MEMORY_COUNTER(ZenHttpClientSerializedBytes, TEXT("ZenClient/SerializedBytes"));
 
-FStorageServerConnection::~FStorageServerConnection()
-{
-	if (StatsObject)
-	{
-		StatsObject->StopDrawing();
-		StatsObject->RemoveFromRoot();
-	}
-}
-
 bool FStorageServerConnection::Initialize(TArrayView<const FString> HostAddresses, const int32 Port, const FAnsiStringView& InBaseURI)
 {
 	BaseURI = InBaseURI;
@@ -48,10 +38,6 @@ bool FStorageServerConnection::Initialize(TArrayView<const FString> HostAddresse
 		CurrentHostAddr = HostAddress;
 		if (HandshakeRequest())
 		{
-			if (StatsObject)
-			{
-				StatsObject->SetHostAddress(CurrentHostAddr);
-			}
 			return true;
 		}
 	}
@@ -207,6 +193,14 @@ void FStorageServerConnection::ReadChunkRequestAsync(
 	}, *ResourceBuilder);
 }
 
+void FStorageServerConnection::GetAndResetStats(IStorageServerPlatformFile::FConnectionStats& OutStats)
+{
+	OutStats.AccumulatedBytes = AccumulatedBytes.exchange(0, std::memory_order_relaxed);
+	OutStats.RequestCount = RequestCount.exchange(0, std::memory_order_relaxed);
+	OutStats.MinRequestThroughput = MinRequestThroughput.exchange(DBL_MAX, std::memory_order_relaxed);
+	OutStats.MaxRequestThroughput = MaxRequestThroughput.exchange(-DBL_MAX, std::memory_order_relaxed);
+}
+
 TArray<FString> FStorageServerConnection::SortHostAddressesByLocalSubnet(TArrayView<const FString> HostAddresses, const int32 Port)
 {
 	bool bAllArePlatformSocketAddresses = true;
@@ -544,17 +538,24 @@ uint64 FStorageServerConnection::GetCompressedOffset(const FCompressedBuffer& Bu
 
 void FStorageServerConnection::AddTimingInstance(const double Duration, const uint64 Bytes)
 {
-	if (StatsObject == nullptr && UObjectInitialized())
+	if ((Duration >= 0.0))
 	{
-		StatsObject = NewObject<UDebugStorageServerConnection>();
-		StatsObject->SetHostAddress(CurrentHostAddr);
-		StatsObject->AddToRoot();
-		StatsObject->StartDrawing();
-	}
+		double tr = ((double)(Bytes * 8) / Duration) / 1000000.0; //Mbps
 
-	if (StatsObject)
-	{
-		StatsObject->AddTimingInstance(Duration, Bytes);
+		AccumulatedBytes.fetch_add(Bytes, std::memory_order_relaxed);
+		RequestCount.fetch_add(1, std::memory_order_relaxed);
+
+		double MinTemp = MinRequestThroughput.load(std::memory_order_relaxed);
+		while (!MinRequestThroughput.compare_exchange_weak(MinTemp, FMath::Min(MinTemp, tr), std::memory_order_relaxed))
+		{
+			MinTemp = MinRequestThroughput.load(std::memory_order_relaxed);
+		}
+
+		double MaxTemp = MaxRequestThroughput.load(std::memory_order_relaxed);
+		while (!MaxRequestThroughput.compare_exchange_weak(MaxTemp, FMath::Max(MaxTemp, tr), std::memory_order_relaxed))
+		{
+			MaxTemp = MaxRequestThroughput.load(std::memory_order_relaxed);
+		}
 	}
 }
 
diff --git a/Engine/Source/Runtime/StorageServerClient/Private/StorageServerConnection.h b/Engine/Source/Runtime/StorageServerClient/Private/StorageServerConnection.h
--- a/Engine/Source/Runtime/StorageServerClient/Private/StorageServerConnection.h
+++ b/Engine/Source/Runtime/StorageServerClient/Private/StorageServerConnection.h
@@ -12,8 +12,8 @@
 #include "SocketTypes.h"
 #include "HAL/PlatformTime.h"
 #include "StorageServerHttpClient.h"
-#include "DebugStorageServerConnection.h"
 #include "IO/IoChunkId.h"
+#include "IStorageServerPlatformFile.h"
 
 #if !UE_BUILD_SHIPPING
 
@@ -25,7 +25,7 @@ class FStorageServerConnection
 {
 public:
 	FStorageServerConnection() = default;
-	~FStorageServerConnection();
+	~FStorageServerConnection() = default;
 
 	bool Initialize(TArrayView<const FString> HostAddresses, const int32 Port, const FAnsiStringView& InBaseURI);
 
@@ -48,16 +48,23 @@ public:
 		TFunctionRef<void(TIoStatusOr<FIoBuffer> Data)> OnResponse
 	);
 
-	FString GetHostAddr() const
+	FStringView GetHostAddr() const
 	{
 		return CurrentHostAddr;
 	}
 
+	void GetAndResetStats(IStorageServerPlatformFile::FConnectionStats& OutStats);
+
 private:
 	TUniquePtr<IStorageServerHttpClient> HttpClient;
 	FAnsiString BaseURI;
 	FString CurrentHostAddr;
-	UDebugStorageServerConnection* StatsObject = nullptr;
+
+	// Stats
+	std::atomic<uint64> AccumulatedBytes = 0;
+	std::atomic<uint32> RequestCount = 0;
+	std::atomic<double> MinRequestThroughput = DBL_MAX;
+	std::atomic<double> MaxRequestThroughput = -DBL_MAX;
 
 	TArray<FString> SortHostAddressesByLocalSubnet(TArrayView<const FString> HostAddresses, const int32 Port);
 	static bool IsPlatformSocketAddress(const FString Address); 
diff --git a/Engine/Source/Runtime/StorageServerClient/Private/StorageServerPlatformFile.cpp b/Engine/Source/Runtime/StorageServerClient/Private/StorageServerPlatformFile.cpp
--- a/Engine/Source/Runtime/StorageServerClient/Private/StorageServerPlatformFile.cpp
+++ b/Engine/Source/Runtime/StorageServerClient/Private/StorageServerPlatformFile.cpp
@@ -5,7 +5,6 @@
 #include "CookOnTheFly.h"
 #include "CookOnTheFlyPackageStore.h"
 #include "HAL/FileManagerGeneric.h"
-#include "HAL/IPlatformFileModule.h"
 #include "Misc/App.h"
 #include "Misc/CommandLine.h"
 #include "Misc/CoreDelegates.h"
@@ -303,15 +302,22 @@ TUniquePtr<FArchive> FStorageServerPlatformFile::TryFindProjectStoreMarkerFile(I
 		return nullptr;
 	}
 
-	FString RelativeStagedPath = TEXT("../../../");
-	FString RootPath = FPaths::RootDir();
-	FString PlatformName = FPlatformProperties::PlatformName();
-	FString CookedOutputPath = FPaths::Combine(FPaths::ProjectDir(), TEXT("Saved"), TEXT("Cooked"), PlatformName);
-
 	TArray<FString> PotentialProjectStorePaths;
-	PotentialProjectStorePaths.Add(RelativeStagedPath);
-	PotentialProjectStorePaths.Add(CookedOutputPath);
-	PotentialProjectStorePaths.Add(RootPath);
+	if (CustomProjectStorePath.IsEmpty())
+	{
+		FString RelativeStagedPath = TEXT("../../../");
+		FString RootPath = FPaths::RootDir();
+		FString PlatformName = FPlatformProperties::PlatformName();
+		FString CookedOutputPath = FPaths::Combine(FPaths::ProjectDir(), TEXT("Saved"), TEXT("Cooked"), PlatformName);
+
+		PotentialProjectStorePaths.Add(RelativeStagedPath);
+		PotentialProjectStorePaths.Add(CookedOutputPath);
+		PotentialProjectStorePaths.Add(RootPath);
+	}
+	else
+	{
+		PotentialProjectStorePaths.Add(CustomProjectStorePath);
+	}
 
 	for (const FString& ProjectStorePath : PotentialProjectStorePaths)
 	{
@@ -468,7 +474,10 @@ bool FStorageServerPlatformFile::Initialize(IPlatformFile* Inner, const TCHAR* C
 					const TSharedPtr<FJsonObject>& ZenServerObject = *ZenServerObjectPtr;
 					ServerProject = ZenServerObject->GetStringField(TEXT("projectid"));
 					ServerPlatform = ZenServerObject->GetStringField(TEXT("oplogid"));
-					BaseURI = ZenServerObject->GetStringField(TEXT("baseuri"));
+					if (!ZenServerObject->TryGetStringField(TEXT("baseuri"), BaseURI))
+					{
+						BaseURI.Empty();
+					}
 					UE_LOG(LogStorageServerPlatformFile, Display, TEXT("Using settings from ue.projectstore: ServerProject='%s' and ServerPlatform='%s'"), *ServerProject, *ServerPlatform);
 				}
 			}
@@ -508,23 +517,27 @@ void FStorageServerPlatformFile::InitializeAfterProjectFilePath()
 	{
 		if (SendGetFileListMessage())
 		{
-			FIoDispatcher& IoDispatcher = FIoDispatcher::Get();
-			TSharedRef<FStorageServerIoDispatcherBackend> IoDispatcherBackend = MakeShared<FStorageServerIoDispatcherBackend>(*Connection.Get());
-			IoDispatcher.Mount(IoDispatcherBackend);
-#if WITH_COTF
-			if (CookOnTheFlyServerConnection)
+			if (bAllowPackageIo)
 			{
-				FPackageStore::Get().Mount(MakeShared<FCookOnTheFlyPackageStoreBackend>(*CookOnTheFlyServerConnection.Get()));
-			}
-			else
+				FIoDispatcher& IoDispatcher = FIoDispatcher::Get();
+				TSharedRef<FStorageServerIoDispatcherBackend> IoDispatcherBackend = MakeShared<FStorageServerIoDispatcherBackend>(*Connection.Get());
+				IoDispatcher.Mount(IoDispatcherBackend);
+#if WITH_COTF
+				if (CookOnTheFlyServerConnection)
+				{
+					FPackageStore::Get().Mount(MakeShared<FCookOnTheFlyPackageStoreBackend>(*CookOnTheFlyServerConnection.Get()));
+				}
+				else
 #endif
-			{
-				FPackageStore::Get().Mount(MakeShared<FStorageServerPackageStoreBackend>(*Connection.Get()));
+				{
+					FPackageStore::Get().Mount(MakeShared<FStorageServerPackageStoreBackend>(*Connection.Get()));
+				}
 			}
 		}
 		else
 		{
-			UE_LOG(LogStorageServerPlatformFile, Fatal, TEXT("Failed to get file list from Zen at '%s'"), *Connection->GetHostAddr());
+			FStringView HostAddr = Connection->GetHostAddr();
+			UE_LOG(LogStorageServerPlatformFile, Fatal, TEXT("Failed to get file list from Zen at '%.*s'"), HostAddr.Len(), HostAddr.GetData());
 		}
 	}
 	else
@@ -1049,6 +1062,16 @@ bool FStorageServerPlatformFile::SendMessageToServer(const TCHAR* Message, IPlat
 	return false;
 }
 
+FStringView FStorageServerPlatformFile::GetHostAddr() const
+{
+	return Connection->GetHostAddr();
+}
+
+void FStorageServerPlatformFile::GetAndResetConnectionStats(FConnectionStats& OutStats)
+{
+	return Connection->GetAndResetStats(OutStats);
+}
+
 #if WITH_COTF
 void FStorageServerPlatformFile::OnCookOnTheFlyMessage(const UE::Cook::FCookOnTheFlyMessage& Message)
 {
@@ -1081,18 +1104,4 @@ void FStorageServerPlatformFile::OnCookOnTheFlyMessage(const UE::Cook::FCookOnTh
 }
 #endif
 
-class FStorageServerClientFileModule
-	: public IPlatformFileModule
-{
-public:
-
-	virtual IPlatformFile* GetPlatformFile() override
-	{
-		static TUniquePtr<IPlatformFile> AutoDestroySingleton = MakeUnique<FStorageServerPlatformFile>();
-		return AutoDestroySingleton.Get();
-	}
-};
-
-IMPLEMENT_MODULE(FStorageServerClientFileModule, StorageServerClient);
-
 #endif
diff --git a/Engine/Source/Runtime/StorageServerClient/Private/StorageServerPlatformFile.h b/Engine/Source/Runtime/StorageServerClient/Private/StorageServerPlatformFile.h
--- a/Engine/Source/Runtime/StorageServerClient/Private/StorageServerPlatformFile.h
+++ b/Engine/Source/Runtime/StorageServerClient/Private/StorageServerPlatformFile.h
@@ -6,6 +6,7 @@
 #include "GenericPlatform/GenericPlatformFile.h"
 #include "Containers/StringView.h"
 #include "IO/IoDispatcher.h"
+#include "IStorageServerPlatformFile.h"
 
 #if !UE_BUILD_SHIPPING
 
@@ -59,7 +60,7 @@ private:
 };
 
 class FStorageServerPlatformFile
-	: public IPlatformFile
+	: public IStorageServerPlatformFile
 {
 public:
 	FStorageServerPlatformFile();
@@ -105,6 +106,17 @@ public:
 	virtual FString ConvertToAbsolutePathForExternalAppForRead(const TCHAR* Filename) override;
 	virtual bool SendMessageToServer(const TCHAR* Message, IPlatformFile::IFileServerMessageHandler* Handler) override;
 
+	FStringView GetHostAddr() const override;
+	void GetAndResetConnectionStats(FConnectionStats& OutStats) override;
+
+	void SetAllowPackageIo(bool bInAllowPackageIo)
+	{
+		bAllowPackageIo = bInAllowPackageIo;
+	}
+	void SetCustomProjectStorePath(FStringView InProjectStorePath)
+	{
+		CustomProjectStorePath = InProjectStorePath;
+	}
 private:
 	friend class FStorageServerFileHandle;
 
@@ -122,6 +134,7 @@ private:
 	TUniquePtr<FArchive> TryFindProjectStoreMarkerFile(IPlatformFile* Inner) const;
 	FAnsiString MakeBaseURI();
 
+	FString CustomProjectStorePath;
 	TSet<FName> ExcludedNonServerExtensions;
 	TSet<FName> AssumedImmutableTimeStampExtensions;
 	IPlatformFile* LowerLevel = nullptr;
@@ -137,6 +150,7 @@ private:
 	FString BaseURI;
 	mutable TArray<FString> HostAddrs;
 	mutable uint16 HostPort = 8558;
+	bool bAllowPackageIo = true;
 };
 
 #endif
diff --git a/Engine/Source/Runtime/StorageServerClient/StorageServerClient.Build.cs b/Engine/Source/Runtime/StorageServerClient/StorageServerClient.Build.cs
--- a/Engine/Source/Runtime/StorageServerClient/StorageServerClient.Build.cs
+++ b/Engine/Source/Runtime/StorageServerClient/StorageServerClient.Build.cs
@@ -6,15 +6,14 @@ public class StorageServerClient : ModuleRules
 {
 	public StorageServerClient(ReadOnlyTargetRules Target) : base(Target)
 	{
-        PublicDependencyModuleNames.AddRange(
-            new string[] {
-                "Core",
-                "CoreUObject",
-                "Engine",        // only for on-screen debug messages
-                "Sockets",
-                "CookOnTheFly",
-                "Json"
-            }
-        );
-    }
+		PrivateDependencyModuleNames.AddRange(
+			new string[] {
+				"Core",
+				"CoreUObject",
+				"Sockets",
+				"CookOnTheFly",
+				"Json"
+			}
+		);
+	}
 }
diff --git a/Engine/Source/Runtime/StorageServerClientDebug/Private/StorageServerConnectionDebug.cpp b/Engine/Source/Runtime/StorageServerClientDebug/Private/StorageServerConnectionDebug.cpp
new file mode 100644
--- /dev/null
+++ b/Engine/Source/Runtime/StorageServerClientDebug/Private/StorageServerConnectionDebug.cpp
@@ -0,0 +1,254 @@
+// Copyright Epic Games, Inc. All Rights Reserved.
+
+#include "StorageServerConnectionDebug.h"
+
+#include "Debug/DebugDrawService.h"
+#include "Engine/GameEngine.h"
+#include "Modules/ModuleInterface.h"
+#include "Modules/ModuleManager.h"
+#include "ProfilingDebugging/CsvProfiler.h"
+#include "Templates/UniquePtr.h"
+#include "StorageServerClientModule.h"
+
+#if !UE_BUILD_SHIPPING
+
+CSV_DEFINE_CATEGORY(ZenServerStats, true);
+
+CSV_DEFINE_STAT(ZenServerStats, ThroughputMbps);
+CSV_DEFINE_STAT(ZenServerStats, MaxReqThroughputMbps);
+CSV_DEFINE_STAT(ZenServerStats, MinReqThroughputMbps);
+CSV_DEFINE_STAT(ZenServerStats, RequestCountPerSec);
+
+bool UStorageServerConnectionDebug::ShowGraphs = false;
+
+void UStorageServerConnectionDebug::StartDrawing()
+{
+	if (DrawHandle.IsValid())
+	{
+		return;
+	}
+	DrawHandle = UDebugDrawService::Register(TEXT("Game"),FDebugDrawDelegate::CreateUObject(this, &UStorageServerConnectionDebug::Draw));
+}
+
+void UStorageServerConnectionDebug::StopDrawing()
+{
+	if (!DrawHandle.IsValid())
+	{
+		return;
+	}
+
+	UDebugDrawService::Unregister(DrawHandle);
+	DrawHandle.Reset();
+}
+
+void UStorageServerConnectionDebug::Draw(UCanvas* Canvas, APlayerController*)
+{
+	static constexpr double FrameSeconds = 1.0;
+	static constexpr float  ViewXRel = 0.2f;
+	static constexpr float  ViewYRel = 0.12f;
+	static constexpr float  ViewWidthRel = 0.4f;
+	static constexpr float  ViewHeightRel = 0.18f;
+	static constexpr double TextHeight = 16.0;
+	static constexpr int    OneMinuteSeconds = 60;
+	static constexpr double WidthSeconds = OneMinuteSeconds * 0.25;
+	static constexpr double MaxHeightScaleThroughput = 6000;
+	static constexpr double MaxHeightScaleRequest = 5000;
+	static constexpr int	LineThickness = 3;
+	static double			HeightScaleThroughput = MaxHeightScaleThroughput;
+	static double			HeightScaleRequest = MaxHeightScaleRequest;
+	
+	double StatsTimeNow = FPlatformTime::Seconds();
+	double Duration = StatsTimeNow - UpdateStatsTime;
+
+	static double MaxReqThroughput = 0.0;
+	static double MinReqThroughput = 0.0;
+	static uint32 ReqCount = 0;
+	static double Throughput = 0.0;
+
+	//Persistent debug message and CSV stats
+	if ((Duration > UpdateStatsTimer) && (GEngine) && StorageServerPlatformFile)
+	{
+		UpdateStatsTime = StatsTimeNow;
+
+		IStorageServerPlatformFile::FConnectionStats Stats;
+		StorageServerPlatformFile->GetAndResetConnectionStats(Stats);
+		if (Stats.MaxRequestThroughput > Stats.MinRequestThroughput)
+		{
+			MaxReqThroughput = Stats.MaxRequestThroughput;
+			MinReqThroughput = Stats.MinRequestThroughput;
+		
+			Throughput = ((double)(Stats.AccumulatedBytes * 8) / Duration) / 1000000.0; //Mbps
+			ReqCount = ceil((double)Stats.RequestCount / Duration);
+		}
+
+		FString ZenConnectionDebugMsg;
+		ZenConnectionDebugMsg = FString::Printf(TEXT("ZenServer streaming from %s [%.2fMbps]"), *HostAddress, Throughput);
+		GEngine->AddOnScreenDebugMessage((uint64)this, 86400.0f, FColor::White, ZenConnectionDebugMsg, false);
+		
+		History.push_back({ StatsTimeNow, MaxReqThroughput, MinReqThroughput, Throughput, ReqCount });
+	}
+
+	while (!History.empty() && StatsTimeNow - History.front().Time > WidthSeconds)
+	{
+		History.erase(History.begin());
+	}
+
+	//CSV stats need to be written per frame
+	CSV_CUSTOM_STAT_DEFINED(ThroughputMbps, Throughput, ECsvCustomStatOp::Set);
+	CSV_CUSTOM_STAT_DEFINED(MaxReqThroughputMbps, MaxReqThroughput, ECsvCustomStatOp::Set);
+	CSV_CUSTOM_STAT_DEFINED(MinReqThroughputMbps, MinReqThroughput, ECsvCustomStatOp::Set);
+	CSV_CUSTOM_STAT_DEFINED(RequestCountPerSec, (int32)ReqCount, ECsvCustomStatOp::Set);
+
+	if (ShowGraphs)
+	{
+		int ViewX = (int)(ViewXRel * Canvas->ClipX);
+		int ViewY = (int)(ViewYRel * Canvas->ClipY);
+		int ViewWidth = (int)(ViewWidthRel * Canvas->ClipX);;
+		int ViewHeight = (int)(ViewHeightRel * Canvas->ClipY);;
+		double PixelsPerSecond = ViewWidth / WidthSeconds;
+
+		auto DrawLine =
+			[Canvas](double X0, double Y0, double X1, double Y1, const FLinearColor& Color, double Thickness)
+			{
+				FCanvasLineItem Line{ FVector2D{X0, Y0}, FVector2D{X1, Y1} };
+				Line.SetColor(Color);
+				Line.LineThickness = Thickness;
+				Canvas->DrawItem(Line);
+			};
+
+		auto DrawString =
+			[Canvas](const FString& Str, int X, int Y, bool bCentre = true)
+			{
+				FCanvasTextItem Text{ FVector2D(X, Y), FText::FromString(Str), GEngine->GetTinyFont(), FLinearColor::Yellow };
+				Text.EnableShadow(FLinearColor::Black);
+				Text.bCentreX = bCentre;
+				Text.bCentreY = bCentre;
+				Canvas->DrawItem(Text);
+			};
+
+		double MaxValueInHistory = 0.0;
+
+		if (History.size())
+		{
+			ViewY += TextHeight;
+			DrawString(FString::Printf(TEXT("Request Throughput MIN/MAX: [%.2f] / [%.2f] Mbps"), History[History.size()-1].MinRequestThroughput, History[History.size()-1].MaxRequestThroughput), ViewX, ViewY, false);
+			ViewY += TextHeight;
+		}
+
+		//FIRST GRAPH
+		MaxValueInHistory = 0.0;
+		double HeightScale = HeightScaleThroughput;
+		ViewY += TextHeight;
+		{ // draw graph edges + label
+			const FLinearColor Color = FLinearColor::White;
+			DrawLine(ViewX, ViewY + ViewHeight, ViewX + ViewWidth, ViewY + ViewHeight, Color, 1);
+			DrawLine(ViewX, ViewY, ViewX, ViewY + ViewHeight, Color, 1);
+			DrawLine(ViewX + ViewWidth, ViewY, ViewX + ViewWidth, ViewY + ViewHeight, Color, 1);
+			DrawString(TEXT("ZenServer Throughput Mbps"), ViewX, ViewY + ViewHeight + 10, false);
+		}
+
+		for (int I = History.size() - 1; I >= 0; --I)
+		{
+			const auto& Item = History[I];
+			const double       X = ViewX + ViewWidth - PixelsPerSecond * (StatsTimeNow - Item.Time);
+			const double       H = FMath::Min(ViewHeight, ViewHeight * (Item.Throughput / HeightScale));
+			const double       Y = ViewY + ViewHeight - H;
+			const FLinearColor Color = FLinearColor::Yellow;
+			
+			DrawLine(X, ViewY + ViewHeight - 1, X, Y, Color, LineThickness);
+			DrawString(FString::Printf(TEXT("%.2f"), Item.Throughput), X, Y - 11);
+
+			if (Item.Throughput > MaxValueInHistory)
+				MaxValueInHistory = Item.Throughput;
+		}
+		HeightScaleThroughput = FMath::Min(MaxHeightScaleThroughput, FMath::Max(MaxValueInHistory, 1.0));
+
+		//SECOND GRAPH
+		MaxValueInHistory = 0.0;
+		ViewY += ViewHeight + (TextHeight * 2) ;
+		HeightScale = HeightScaleRequest;
+
+		{ // draw graph edges + label
+			const FLinearColor Color = FLinearColor::White;
+			DrawLine(ViewX, ViewY + ViewHeight, ViewX + ViewWidth, ViewY + ViewHeight, Color, 1);
+			DrawLine(ViewX, ViewY, ViewX, ViewY + ViewHeight, Color, 1);
+			DrawLine(ViewX + ViewWidth, ViewY, ViewX + ViewWidth, ViewY + ViewHeight, Color, 1);
+			DrawString(TEXT("ZenServer Request/Sec Count"), ViewX, ViewY + ViewHeight + 10, false);
+		}
+
+		for (int I = History.size() - 1; I >= 0; --I)
+		{
+			const auto& Item = History[I];
+			const double       X = ViewX + ViewWidth - PixelsPerSecond * (StatsTimeNow - Item.Time);
+			const double       H = FMath::Min(ViewHeight, ViewHeight * (Item.RequestCount / HeightScale));
+			const double       Y = ViewY + ViewHeight - H;
+			const FLinearColor Color = FLinearColor::Gray;
+
+			DrawLine(X, ViewY + ViewHeight - 1, X, Y, Color, LineThickness);
+			DrawString(FString::Printf(TEXT("%.d"), Item.RequestCount), X, Y - 11);
+
+			if (Item.RequestCount > MaxValueInHistory)
+				MaxValueInHistory = Item.RequestCount;
+		}
+		HeightScaleRequest = FMath::Min(MaxHeightScaleRequest, FMath::Max(MaxValueInHistory, 1.0));
+	}
+}
+
+void UStorageServerConnectionDebug::ShowGraph(FOutputDevice&)
+{
+	ShowGraphs = true;
+}
+
+void UStorageServerConnectionDebug::HideGraph(FOutputDevice&)
+{
+	ShowGraphs = false;
+}
+
+static FAutoConsoleCommandWithOutputDevice
+	GShowDebugConnectionStatsCmd(
+		TEXT("r.ZenServerStatsShow"),
+		TEXT("Show ZenServer Stats Graph."),
+		FConsoleCommandWithOutputDeviceDelegate::CreateStatic(&UStorageServerConnectionDebug::ShowGraph));
+
+static FAutoConsoleCommandWithOutputDevice
+	GHideDebugConnectionStatsCmd(
+		TEXT("r.ZenServerStatsHide"),
+		TEXT("Hide ZenServer Stats Graph."),
+		FConsoleCommandWithOutputDeviceDelegate::CreateStatic(&UStorageServerConnectionDebug::HideGraph));
+
+
+class FStorageServerClientDebugModule
+	: public IModuleInterface
+{
+public:
+	virtual void StartupModule() override
+	{
+		FCoreDelegates::OnPostEngineInit.AddLambda([this]
+		{
+			if (IStorageServerPlatformFile* StorageServerPlatformFile = IStorageServerClientModule::FindStorageServerPlatformFile())
+			{
+				ConnectionDebug = NewObject<UStorageServerConnectionDebug>();
+				ConnectionDebug->SetPlatformFile(StorageServerPlatformFile);
+				ConnectionDebug->AddToRoot();
+				ConnectionDebug->StartDrawing();
+			}
+		});
+	}
+
+	virtual void ShutdownModule() override
+	{
+		if (ConnectionDebug)
+		{
+			ConnectionDebug->SetPlatformFile(nullptr);
+			ConnectionDebug->StopDrawing();
+			ConnectionDebug->RemoveFromRoot();
+			ConnectionDebug = nullptr;
+		}
+	}
+
+	UStorageServerConnectionDebug* ConnectionDebug = nullptr;
+};
+
+IMPLEMENT_MODULE(FStorageServerClientDebugModule, StorageServerClientDebug);
+
+#endif
\ No newline at end of file
diff --git a/Engine/Source/Runtime/StorageServerClientDebug/Private/StorageServerConnectionDebug.h b/Engine/Source/Runtime/StorageServerClientDebug/Private/StorageServerConnectionDebug.h
new file mode 100644
--- /dev/null
+++ b/Engine/Source/Runtime/StorageServerClientDebug/Private/StorageServerConnectionDebug.h
@@ -0,0 +1,61 @@
+// Copyright Epic Games, Inc. All Rights Reserved.
+
+#pragma once
+
+#include "CoreMinimal.h"
+#include "Engine/Canvas.h"
+#include <vector>
+#include "IStorageServerPlatformFile.h"
+
+#include "StorageServerConnectionDebug.generated.h"
+
+UCLASS()
+class UStorageServerConnectionDebug : public UObject
+{
+	GENERATED_BODY()
+
+public:
+	void StartDrawing();
+	void StopDrawing();
+
+	void SetPlatformFile(IStorageServerPlatformFile* InStorageServerPlatformFile)
+	{
+		StorageServerPlatformFile = InStorageServerPlatformFile;
+		if (StorageServerPlatformFile != nullptr)
+		{
+			HostAddress = InStorageServerPlatformFile->GetHostAddr();
+		}
+		else
+		{
+			HostAddress.Reset();
+		}
+	}
+
+	static void ShowGraph(FOutputDevice&);
+	static void HideGraph(FOutputDevice&);
+
+private:
+	void Draw(UCanvas* Canvas, class APlayerController* PC);
+
+	struct HistoryItem
+	{
+		double Time;
+		double MaxRequestThroughput;
+		double MinRequestThroughput;
+		double Throughput;
+		uint32 RequestCount;
+	};
+
+	std::vector<HistoryItem> History = {{0, 0, 0, 0, 0}};
+
+	FDelegateHandle DrawHandle;
+
+	static constexpr float UpdateStatsTimer = 1.0;
+	double UpdateStatsTime = 0.0;
+
+	IStorageServerPlatformFile* StorageServerPlatformFile = nullptr;
+	FString HostAddress;
+
+	static bool ShowGraphs;
+
+};
diff --git a/Engine/Source/Runtime/StorageServerClientDebug/StorageServerClientDebug.Build.cs b/Engine/Source/Runtime/StorageServerClientDebug/StorageServerClientDebug.Build.cs
new file mode 100644
--- /dev/null
+++ b/Engine/Source/Runtime/StorageServerClientDebug/StorageServerClientDebug.Build.cs
@@ -0,0 +1,20 @@
+// Copyright Epic Games, Inc. All Rights Reserved.
+
+using UnrealBuildTool;
+
+public class StorageServerClientDebug : ModuleRules
+{
+	public StorageServerClientDebug(ReadOnlyTargetRules Target) : base(Target)
+	{
+		PrivateDependencyModuleNames.AddRange(
+			new string[] {
+				"Core",
+				"CoreUObject",
+				"Engine",		// For on-screen debug messages
+				"Sockets",
+				"StorageServerClient",
+				"Json"
+			}
+		);
+	}
+}
