Fix for shipping build breakage.

diff --git a/Engine/Source/Runtime/StorageServerClient/Internal/IStorageServerPlatformFile.h b/Engine/Source/Runtime/StorageServerClient/Internal/IStorageServerPlatformFile.h
--- a/Engine/Source/Runtime/StorageServerClient/Internal/IStorageServerPlatformFile.h
+++ b/Engine/Source/Runtime/StorageServerClient/Internal/IStorageServerPlatformFile.h
@@ -5,8 +5,6 @@
 #include "Containers/StringFwd.h"
 #include "HAL/PlatformFile.h"
 
-#if !UE_BUILD_SHIPPING
-
 class IStorageServerPlatformFile : public IPlatformFile
 {
 public:
@@ -23,5 +21,3 @@ public:
 
 	virtual void GetAndResetConnectionStats(FConnectionStats& OutStats) = 0;
 };
-
-#endif // !UE_BUILD_SHIPPING
diff --git a/Engine/Source/Runtime/StorageServerClientDebug/Private/StorageServerConnectionDebug.cpp b/Engine/Source/Runtime/StorageServerClientDebug/Private/StorageServerConnectionDebug.cpp
--- a/Engine/Source/Runtime/StorageServerClientDebug/Private/StorageServerConnectionDebug.cpp
+++ b/Engine/Source/Runtime/StorageServerClientDebug/Private/StorageServerConnectionDebug.cpp
@@ -10,8 +10,6 @@
 #include "Templates/UniquePtr.h"
 #include "StorageServerClientModule.h"
 
-#if !UE_BUILD_SHIPPING
-
 CSV_DEFINE_CATEGORY(ZenServerStats, true);
 
 CSV_DEFINE_STAT(ZenServerStats, ThroughputMbps);
@@ -221,6 +219,7 @@ class FStorageServerClientDebugModule
 	: public IModuleInterface
 {
 public:
+#if !UE_BUILD_SHIPPING
 	virtual void StartupModule() override
 	{
 		FCoreDelegates::OnPostEngineInit.AddLambda([this]
@@ -245,10 +244,9 @@ public:
 			ConnectionDebug = nullptr;
 		}
 	}
+#endif // !UE_BUILD_SHIPPING
 
 	UStorageServerConnectionDebug* ConnectionDebug = nullptr;
 };
 
 IMPLEMENT_MODULE(FStorageServerClientDebugModule, StorageServerClientDebug);
-
-#endif
\ No newline at end of file
