Fix invalid access to dangling pointer during Shutdow of FStorageServerClientDebugModule

diff --git a/Engine/Source/Runtime/StorageServerClientDebug/Private/StorageServerConnectionDebug.cpp b/Engine/Source/Runtime/StorageServerClientDebug/Private/StorageServerConnectionDebug.cpp
--- a/Engine/Source/Runtime/StorageServerClientDebug/Private/StorageServerConnectionDebug.cpp
+++ b/Engine/Source/Runtime/StorageServerClientDebug/Private/StorageServerConnectionDebug.cpp
@@ -236,13 +236,8 @@ public:
 
 	virtual void ShutdownModule() override
 	{
-		if (ConnectionDebug)
-		{
-			ConnectionDebug->SetPlatformFile(nullptr);
-			ConnectionDebug->StopDrawing();
-			ConnectionDebug->RemoveFromRoot();
-			ConnectionDebug = nullptr;
-		}
+		//Since ConnectionDebug is an UObject it's already been freed during shutdown at this point
+		ConnectionDebug = nullptr;
 	}
 #endif // !UE_BUILD_SHIPPING
 
