New stats & debugging for Zen:

New Insights stats:
 - File I/O trace for storage server files (requires running with -trace=default,file)
 - Throughput counters, matching the on-screen display (requires -trace=default,counters)

Debug refactor:
 - Main stat processing moved to ticker
 - Startup stats processed by background thread

CVar changes:
 - On-screen graphs now toggled via zen.showgraphs
 - Added zen.showstats to toggle the on-screen throughput messages

#rnx

diff --git a/Engine/Source/Runtime/StorageServerClient/Private/StorageServerClientModule.cpp b/Engine/Source/Runtime/StorageServerClient/Private/StorageServerClientModule.cpp
--- a/Engine/Source/Runtime/StorageServerClient/Private/StorageServerClientModule.cpp
+++ b/Engine/Source/Runtime/StorageServerClient/Private/StorageServerClientModule.cpp
@@ -44,11 +44,6 @@ IPlatformFile* FStorageServerClientModule::GetPlatformFile()
 	if (!CachedDefaultPlatformFileInstance)
 	{
 		CachedDefaultPlatformFileInstance = DefaultPlatformFileInstance.Get();
-		// Load optional modules.
-		if (FModuleManager::Get().ModuleExists(TEXT("StorageServerClientDebug")))
-		{
-			FModuleManager::Get().LoadModule("StorageServerClientDebug");
-		}
 	}
 	return DefaultPlatformFileInstance.Get();
 }
diff --git a/Engine/Source/Runtime/StorageServerClient/Private/StorageServerConnection.cpp b/Engine/Source/Runtime/StorageServerClient/Private/StorageServerConnection.cpp
--- a/Engine/Source/Runtime/StorageServerClient/Private/StorageServerConnection.cpp
+++ b/Engine/Source/Runtime/StorageServerClient/Private/StorageServerConnection.cpp
@@ -25,7 +25,8 @@
 
 DEFINE_LOG_CATEGORY(LogStorageServerConnection);
 
-TRACE_DECLARE_MEMORY_COUNTER(ZenHttpClientSerializedBytes, TEXT("ZenClient/SerializedBytes"));
+TRACE_DECLARE_INT_COUNTER(ZenHttpClientSerializedBytes, TEXT("ZenClient/SerializedBytes (compressed)"));
+TRACE_DECLARE_INT_COUNTER(ZenHttpClientThroughputBytes, TEXT("ZenClient/ThroughputBytes (decompressed)"));
 
 bool FStorageServerConnection::Initialize(TArrayView<const FString> HostAddresses, const int32 Port, const FAnsiStringView& InBaseURI)
 {
@@ -557,6 +558,8 @@ void FStorageServerConnection::AddTimingInstance(const double Duration, const ui
 			MaxTemp = MaxRequestThroughput.load(std::memory_order_relaxed);
 		}
 	}
+
+	TRACE_COUNTER_ADD(ZenHttpClientThroughputBytes, Bytes);
 }
 
 // TODO revive FStorageServerChunkBatchRequest
diff --git a/Engine/Source/Runtime/StorageServerClient/Private/StorageServerPlatformFile.cpp b/Engine/Source/Runtime/StorageServerClient/Private/StorageServerPlatformFile.cpp
--- a/Engine/Source/Runtime/StorageServerClient/Private/StorageServerPlatformFile.cpp
+++ b/Engine/Source/Runtime/StorageServerClient/Private/StorageServerPlatformFile.cpp
@@ -18,6 +18,7 @@
 #include "Serialization/CompactBinarySerialization.h"
 #include "Serialization/JsonReader.h"
 #include "Serialization/JsonSerializer.h"
+#include "ProfilingDebugging/PlatformFileTrace.h"
 #include "StorageServerConnection.h"
 #include "StorageServerIoDispatcherBackend.h"
 #include "StorageServerPackageStore.h"
@@ -182,10 +183,14 @@ public:
 		, Filename(InFilename)
 		, FileSize(InFileSize)
 	{
+		TRACE_PLATFORMFILE_BEGIN_OPEN(*FString::Printf(TEXT("zen:%s"), InFilename));
+		TRACE_PLATFORMFILE_END_OPEN(this);
 	}
 
 	~FStorageServerFileHandle()
 	{
+		TRACE_PLATFORMFILE_BEGIN_CLOSE(this);
+		TRACE_PLATFORMFILE_END_CLOSE(this);
 	}
 
 	virtual int64 Size() override
@@ -224,8 +229,10 @@ public:
 
 	virtual bool Read(uint8* Destination, int64 BytesToRead) override
 	{
+		TRACE_PLATFORMFILE_BEGIN_READ(Destination, this, FilePos, BytesToRead);
 		if (BytesToRead == 0)
 		{
+			TRACE_PLATFORMFILE_END_READ(Destination, 0);
 			return true;
 		}
 
@@ -235,8 +242,10 @@ public:
 			if (BytesRead == BytesToRead)
 			{
 				FilePos += BytesRead;
+				TRACE_PLATFORMFILE_END_READ(Destination, BytesRead);
 				return true;
 			}
+			TRACE_PLATFORMFILE_END_READ(Destination, 0);
 			return false;
 		}
 
@@ -250,6 +259,7 @@ public:
 			if (BytesReadFromBuffer == BytesToRead)
 			{
 				FilePos += BytesReadFromBuffer;
+				TRACE_PLATFORMFILE_END_READ(this, BytesReadFromBuffer);
 				return true;
 			}
 		}
@@ -264,9 +274,11 @@ public:
 		if (BytesReadFromBuffer == BytesToRead)
 		{
 			FilePos += BytesReadFromBuffer;
+			TRACE_PLATFORMFILE_END_READ(Destination, BytesReadFromBuffer);
 			return true;
 		}
 		
+		TRACE_PLATFORMFILE_END_READ(Destination, 0);
 		return false;
 	}
 
@@ -539,6 +551,12 @@ void FStorageServerPlatformFile::InitializeAfterProjectFilePath()
 			FStringView HostAddr = Connection->GetHostAddr();
 			UE_LOG(LogStorageServerPlatformFile, Fatal, TEXT("Failed to get file list from Zen at '%.*s'"), HostAddr.Len(), HostAddr.GetData());
 		}
+
+		// optional debugging module depends on a valid Connection
+		if (FModuleManager::Get().ModuleExists(TEXT("StorageServerClientDebug")))
+		{
+			FModuleManager::Get().LoadModule("StorageServerClientDebug");
+		}
 	}
 	else
 	{
diff --git a/Engine/Source/Runtime/StorageServerClientDebug/Private/StorageServerConnectionDebug.cpp b/Engine/Source/Runtime/StorageServerClientDebug/Private/StorageServerConnectionDebug.cpp
--- a/Engine/Source/Runtime/StorageServerClientDebug/Private/StorageServerConnectionDebug.cpp
+++ b/Engine/Source/Runtime/StorageServerClientDebug/Private/StorageServerConnectionDebug.cpp
@@ -7,8 +7,18 @@
 #include "Modules/ModuleInterface.h"
 #include "Modules/ModuleManager.h"
 #include "ProfilingDebugging/CsvProfiler.h"
+#include "ProfilingDebugging/CountersTrace.h"
 #include "Templates/UniquePtr.h"
 #include "StorageServerClientModule.h"
+#include "HAL/IConsoleManager.h"
+#include "HAL/Runnable.h"
+#include "HAL/RunnableThread.h"
+#include "HAL/Event.h"
+#include "HAL/PlatformProcess.h"
+#include "Misc/ScopeLock.h"
+#include "Misc/CoreDelegates.h"
+
+#if !UE_BUILD_SHIPPING
 
 CSV_DEFINE_CATEGORY(ZenServerStats, true);
 
@@ -17,54 +27,47 @@ CSV_DEFINE_STAT(ZenServerStats, MaxReqThroughputMbps);
 CSV_DEFINE_STAT(ZenServerStats, MinReqThroughputMbps);
 CSV_DEFINE_STAT(ZenServerStats, RequestCountPerSec);
 
-bool UStorageServerConnectionDebug::ShowGraphs = false;
+TRACE_DECLARE_UNCHECKED_FLOAT_COUNTER(ZenClient_ThroughputMbps,       TEXT("ZenClient/ThroughputMbps (decompressed)"));
+TRACE_DECLARE_UNCHECKED_FLOAT_COUNTER(ZenClient_MaxReqThroughputMbps, TEXT("ZenClient/MaxReqThroughputMbps (decompressed)"));
+TRACE_DECLARE_UNCHECKED_FLOAT_COUNTER(ZenClient_MinReqThroughputMbps, TEXT("ZenClient/MinReqThroughputMbps (decompressed)"));
+TRACE_DECLARE_UNCHECKED_INT_COUNTER(  ZenClient_RequestCountPerSec,   TEXT("ZenClient/RequestCountPerSec"));
 
-void UStorageServerConnectionDebug::StartDrawing()
-{
-	if (DrawHandle.IsValid())
-	{
-		return;
-	}
-	DrawHandle = UDebugDrawService::Register(TEXT("Game"),FDebugDrawDelegate::CreateUObject(this, &UStorageServerConnectionDebug::Draw));
-}
 
-void UStorageServerConnectionDebug::StopDrawing()
-{
-	if (!DrawHandle.IsValid())
-	{
-		return;
-	}
+static bool GZenShowGraphs = false;
+static FAutoConsoleVariableRef CVarZenShowGraphs(
+	TEXT("zen.showgraphs"),
+	GZenShowGraphs,
+	TEXT("Show ZenServer Stats Graph"),
+	ECVF_Default
+);
 
-	UDebugDrawService::Unregister(DrawHandle);
-	DrawHandle.Reset();
-}
+static bool GZenShowStats = true;
+static FAutoConsoleVariableRef CVarZenShowStats(
+	TEXT("zen.showstats"),
+	GZenShowStats,
+	TEXT("Show ZenServer Stats"),
+	ECVF_Default
+);
 
-void UStorageServerConnectionDebug::Draw(UCanvas* Canvas, APlayerController*)
+
+namespace
 {
-	static constexpr double FrameSeconds = 1.0;
-	static constexpr float  ViewXRel = 0.2f;
-	static constexpr float  ViewYRel = 0.12f;
-	static constexpr float  ViewWidthRel = 0.4f;
-	static constexpr float  ViewHeightRel = 0.18f;
-	static constexpr double TextHeight = 16.0;
 	static constexpr int    OneMinuteSeconds = 60;
 	static constexpr double WidthSeconds = OneMinuteSeconds * 0.25;
-	static constexpr double MaxHeightScaleThroughput = 6000;
-	static constexpr double MaxHeightScaleRequest = 5000;
-	static constexpr int	LineThickness = 3;
-	static double			HeightScaleThroughput = MaxHeightScaleThroughput;
-	static double			HeightScaleRequest = MaxHeightScaleRequest;
+}
+
+
+bool FStorageServerConnectionDebug::OnTick(float)
+{
+	FScopeLock Lock(&CS);
+
+	static constexpr double FrameSeconds = 1.0;
 	
 	double StatsTimeNow = FPlatformTime::Seconds();
 	double Duration = StatsTimeNow - UpdateStatsTime;
 
-	static double MaxReqThroughput = 0.0;
-	static double MinReqThroughput = 0.0;
-	static uint32 ReqCount = 0;
-	static double Throughput = 0.0;
-
 	//Persistent debug message and CSV stats
-	if ((Duration > UpdateStatsTimer) && (GEngine) && StorageServerPlatformFile)
+	if (Duration > UpdateStatsTimer)
 	{
 		UpdateStatsTime = StatsTimeNow;
 
@@ -79,11 +82,19 @@ void UStorageServerConnectionDebug::Draw(UCanvas* Canvas, APlayerController*)
 			ReqCount = ceil((double)Stats.RequestCount / Duration);
 		}
 
-		FString ZenConnectionDebugMsg;
-		ZenConnectionDebugMsg = FString::Printf(TEXT("ZenServer streaming from %s [%.2fMbps]"), *HostAddress, Throughput);
-		GEngine->AddOnScreenDebugMessage((uint64)this, 86400.0f, FColor::White, ZenConnectionDebugMsg, false);
+		if (GZenShowStats && GEngine)
+		{
+			FString ZenConnectionDebugMsg;
+			ZenConnectionDebugMsg = FString::Printf(TEXT("ZenServer streaming from %s [%.2fMbps]"), *HostAddress, Throughput);
+			GEngine->AddOnScreenDebugMessage((uint64)this, 86400.0f, FColor::White, ZenConnectionDebugMsg, false);
+		}
 		
 		History.push_back({ StatsTimeNow, MaxReqThroughput, MinReqThroughput, Throughput, ReqCount });
+
+		TRACE_COUNTER_SET(ZenClient_ThroughputMbps,       Throughput);
+		TRACE_COUNTER_SET(ZenClient_MaxReqThroughputMbps, MaxReqThroughput);
+		TRACE_COUNTER_SET(ZenClient_MinReqThroughputMbps, MinReqThroughput);
+		TRACE_COUNTER_SET(ZenClient_RequestCountPerSec,   ReqCount);
 	}
 
 	while (!History.empty() && StatsTimeNow - History.front().Time > WidthSeconds)
@@ -91,14 +102,39 @@ void UStorageServerConnectionDebug::Draw(UCanvas* Canvas, APlayerController*)
 		History.erase(History.begin());
 	}
 
-	//CSV stats need to be written per frame
-	CSV_CUSTOM_STAT_DEFINED(ThroughputMbps, Throughput, ECsvCustomStatOp::Set);
-	CSV_CUSTOM_STAT_DEFINED(MaxReqThroughputMbps, MaxReqThroughput, ECsvCustomStatOp::Set);
-	CSV_CUSTOM_STAT_DEFINED(MinReqThroughputMbps, MinReqThroughput, ECsvCustomStatOp::Set);
-	CSV_CUSTOM_STAT_DEFINED(RequestCountPerSec, (int32)ReqCount, ECsvCustomStatOp::Set);
+	//CSV stats need to be written per frame (only send if we're running from the gamethread ticker, not the startup debug thread)
+	if (IsInGameThread())
+	{
+		CSV_CUSTOM_STAT_DEFINED(ThroughputMbps, Throughput, ECsvCustomStatOp::Set);
+		CSV_CUSTOM_STAT_DEFINED(MaxReqThroughputMbps, MaxReqThroughput, ECsvCustomStatOp::Set);
+		CSV_CUSTOM_STAT_DEFINED(MinReqThroughputMbps, MinReqThroughput, ECsvCustomStatOp::Set);
+		CSV_CUSTOM_STAT_DEFINED(RequestCountPerSec, (int32)ReqCount, ECsvCustomStatOp::Set);
+	}
+
+	return true;
+}
+
+
+
+void FStorageServerConnectionDebug::OnDraw(UCanvas* Canvas, APlayerController*)
+{
+	FScopeLock Lock(&CS);
+
+	static constexpr float  ViewXRel = 0.2f;
+	static constexpr float  ViewYRel = 0.12f;
+	static constexpr float  ViewWidthRel = 0.4f;
+	static constexpr float  ViewHeightRel = 0.18f;
+	static constexpr double TextHeight = 16.0;
+	static constexpr double MaxHeightScaleThroughput = 6000;
+	static constexpr double MaxHeightScaleRequest = 5000;
+	static constexpr int	LineThickness = 3;
+	static double			HeightScaleThroughput = MaxHeightScaleThroughput;
+	static double			HeightScaleRequest = MaxHeightScaleRequest;
 
-	if (ShowGraphs)
+	if (GZenShowGraphs)
 	{
+		double StatsTimeNow = FPlatformTime::Seconds();
+
 		int ViewX = (int)(ViewXRel * Canvas->ClipX);
 		int ViewY = (int)(ViewYRel * Canvas->ClipY);
 		int ViewWidth = (int)(ViewWidthRel * Canvas->ClipX);;
@@ -192,56 +228,109 @@ void UStorageServerConnectionDebug::Draw(UCanvas* Canvas, APlayerController*)
 	}
 }
 
-void UStorageServerConnectionDebug::ShowGraph(FOutputDevice&)
-{
-	ShowGraphs = true;
-}
-
-void UStorageServerConnectionDebug::HideGraph(FOutputDevice&)
-{
-	ShowGraphs = false;
-}
-
-static FAutoConsoleCommandWithOutputDevice
-	GShowDebugConnectionStatsCmd(
-		TEXT("r.ZenServerStatsShow"),
-		TEXT("Show ZenServer Stats Graph."),
-		FConsoleCommandWithOutputDeviceDelegate::CreateStatic(&UStorageServerConnectionDebug::ShowGraph));
-
-static FAutoConsoleCommandWithOutputDevice
-	GHideDebugConnectionStatsCmd(
-		TEXT("r.ZenServerStatsHide"),
-		TEXT("Hide ZenServer Stats Graph."),
-		FConsoleCommandWithOutputDeviceDelegate::CreateStatic(&UStorageServerConnectionDebug::HideGraph));
-
 
 class FStorageServerClientDebugModule
 	: public IModuleInterface
+	, public FRunnable
 {
 public:
-#if !UE_BUILD_SHIPPING
 	virtual void StartupModule() override
 	{
-		FCoreDelegates::OnPostEngineInit.AddLambda([this]
+		if (IStorageServerPlatformFile* StorageServerPlatformFile = IStorageServerClientModule::FindStorageServerPlatformFile())
 		{
-			if (IStorageServerPlatformFile* StorageServerPlatformFile = IStorageServerClientModule::FindStorageServerPlatformFile())
+			ConnectionDebug.Reset( new FStorageServerConnectionDebug(StorageServerPlatformFile) );
+			OnDrawDebugHandle = UDebugDrawService::Register(TEXT("Game"), FDebugDrawDelegate::CreateRaw(ConnectionDebug.Get(), &FStorageServerConnectionDebug::OnDraw) );
+
+			// start by capturing engine initialization stats on a background thread
+			StartThread();
+
+			// once the engine has initialized, switch to a more lightweight gamethread ticker
+			FCoreDelegates::OnPostEngineInit.AddLambda([this]
 			{
-				ConnectionDebug = NewObject<UStorageServerConnectionDebug>();
-				ConnectionDebug->SetPlatformFile(StorageServerPlatformFile);
-				ConnectionDebug->AddToRoot();
-				ConnectionDebug->StartDrawing();
-			}
-		});
+				StopThread();
+				StartTick();
+			});
+		}
 	}
 
 	virtual void ShutdownModule() override
 	{
-		//Since ConnectionDebug is an UObject it's already been freed during shutdown at this point
-		ConnectionDebug = nullptr;
+		if (ConnectionDebug.IsValid())
+		{
+			StopThread();
+			StopTick();
+			UDebugDrawService::Unregister(OnDrawDebugHandle);
+			ConnectionDebug.Reset();
+		}
 	}
-#endif // !UE_BUILD_SHIPPING
 
-	UStorageServerConnectionDebug* ConnectionDebug = nullptr;
+	void StartThread()
+	{
+		check(!Thread.IsValid());
+		ThreadStopEvent = FPlatformProcess::GetSynchEventFromPool(true);
+		Thread.Reset( FRunnableThread::Create(this, TEXT("StorageServerStartupDebug"), 0, TPri_Lowest) );
+	}
+
+	void StopThread()
+	{
+		if (Thread.IsValid())
+		{
+			Thread.Reset();
+			FPlatformProcess::ReturnSynchEventToPool(ThreadStopEvent);
+			ThreadStopEvent = nullptr;
+		}
+	}
+
+	void StartTick()
+	{
+		check(!TickHandle.IsValid());
+		TickHandle = FTSTicker::GetCoreTicker().AddTicker(FTickerDelegate::CreateRaw(ConnectionDebug.Get(), &FStorageServerConnectionDebug::OnTick));
+	}
+
+	void StopTick()
+	{
+		if (TickHandle.IsValid())
+		{
+			FTSTicker::GetCoreTicker().RemoveTicker(TickHandle);
+			TickHandle.Reset();
+		}
+	}
+
+	// FRunnable interface
+	virtual uint32 Run() override
+	{
+		while(!ThreadStopEvent->Wait(10))
+		{
+			ConnectionDebug->OnTick(0);
+		}
+		return 0;
+	}
+
+	virtual void Stop() override
+	{
+		ThreadStopEvent->Trigger();
+	}
+	// end of FRunnable interface
+
+
+	TUniquePtr<FStorageServerConnectionDebug> ConnectionDebug;
+	FDelegateHandle OnDrawDebugHandle;
+
+	TUniquePtr<FRunnableThread> Thread;
+	FEvent* ThreadStopEvent = nullptr;
+
+	FTSTicker::FDelegateHandle TickHandle;
+
 };
 
 IMPLEMENT_MODULE(FStorageServerClientDebugModule, StorageServerClientDebug);
+
+#else
+
+// shipping stub
+IMPLEMENT_MODULE(FDefaultModuleImpl, StorageServerClientDebug);
+
+#endif // !UE_BUILD_SHIPPING
+
+
+
diff --git a/Engine/Source/Runtime/StorageServerClientDebug/Private/StorageServerConnectionDebug.h b/Engine/Source/Runtime/StorageServerClientDebug/Private/StorageServerConnectionDebug.h
--- a/Engine/Source/Runtime/StorageServerClientDebug/Private/StorageServerConnectionDebug.h
+++ b/Engine/Source/Runtime/StorageServerClientDebug/Private/StorageServerConnectionDebug.h
@@ -4,38 +4,28 @@
 
 #include "CoreMinimal.h"
 #include "Engine/Canvas.h"
+#include "HAL/CriticalSection.h"
 #include <vector>
 #include "IStorageServerPlatformFile.h"
 
-#include "StorageServerConnectionDebug.generated.h"
-
-UCLASS()
-class UStorageServerConnectionDebug : public UObject
+#if !UE_BUILD_SHIPPING
+class FStorageServerConnectionDebug
 {
-	GENERATED_BODY()
-
 public:
-	void StartDrawing();
-	void StopDrawing();
-
-	void SetPlatformFile(IStorageServerPlatformFile* InStorageServerPlatformFile)
+	FStorageServerConnectionDebug( IStorageServerPlatformFile* InStorageServerPlatformFile )
+		: StorageServerPlatformFile(InStorageServerPlatformFile)
+		, HostAddress(InStorageServerPlatformFile->GetHostAddr())
 	{
-		StorageServerPlatformFile = InStorageServerPlatformFile;
-		if (StorageServerPlatformFile != nullptr)
-		{
-			HostAddress = InStorageServerPlatformFile->GetHostAddr();
-		}
-		else
-		{
-			HostAddress.Reset();
-		}
 	}
 
-	static void ShowGraph(FOutputDevice&);
-	static void HideGraph(FOutputDevice&);
+	bool OnTick(float); // FTickerDelegate
+	void OnDraw(UCanvas*, APlayerController*); // FDebugDrawDelegate
 
 private:
-	void Draw(UCanvas* Canvas, class APlayerController* PC);
+	double MaxReqThroughput = 0.0;
+	double MinReqThroughput = 0.0;
+	uint32 ReqCount = 0;
+	double Throughput = 0.0;
 
 	struct HistoryItem
 	{
@@ -48,14 +38,12 @@ private:
 
 	std::vector<HistoryItem> History = {{0, 0, 0, 0, 0}};
 
-	FDelegateHandle DrawHandle;
-
 	static constexpr float UpdateStatsTimer = 1.0;
 	double UpdateStatsTime = 0.0;
 
 	IStorageServerPlatformFile* StorageServerPlatformFile = nullptr;
 	FString HostAddress;
 
-	static bool ShowGraphs;
-
+	FCriticalSection CS;
 };
+#endif // !UE_BUILD_SHIPPING
