Add on-screen graphical indicator that Zen Streaming is active

#rnx

#virtualized

diff --git a/Engine/Config/BaseEngine.ini b/Engine/Config/BaseEngine.ini
--- a/Engine/Config/BaseEngine.ini
+++ b/Engine/Config/BaseEngine.ini
@@ -209,6 +209,7 @@ DefaultPostProcessMaterialName=/Engine/EngineMaterials/DefaultPostProcessMateria
 DefaultFlattenMaterialName=/Engine/EngineMaterials/FlattenMaterial.FlattenMaterial
 DefaultHLODFlattenMaterialName=/Engine/EngineMaterials/FlattenMaterial_VT.FlattenMaterial_VT
 DefaultLandscapeFlattenMaterialName=/Engine/EngineMaterials/FlattenMaterial_WS_Normal_VT.FlattenMaterial_WS_Normal_VT
+DefaultZenStreamingTextureName=/Engine/EngineResources/Development/ZenLogo_64.ZenLogo_64
 TimecodeProviderClassName=None
 ArrowMaterialName=/Engine/EngineMaterials/GizmoMaterial.GizmoMaterial
 ParticleEventManagerClassPath="/Script/Engine.ParticleEventManager"
@@ -2076,6 +2077,7 @@ AutoSaveIndex=0
 +PackagesToBeFullyLoadedAtStartup=/Engine/EngineResources/M_StreamingPause
 +PackagesToBeFullyLoadedAtStartup=/Engine/EngineResources/WhiteSquareTexture
 +PackagesToBeFullyLoadedAtStartup=/Engine/EngineResources/GradientTexture0
++PackagesToBeFullyLoadedAtStartup=/Engine/EngineResources/Development/ZenIcon
 
 [DevOptions.Shaders]
 ; See FShaderCompilingManager for documentation on what these do
diff --git a/Engine/Source/Runtime/Engine/Classes/Engine/Engine.h b/Engine/Source/Runtime/Engine/Classes/Engine/Engine.h
--- a/Engine/Source/Runtime/Engine/Classes/Engine/Engine.h
+++ b/Engine/Source/Runtime/Engine/Classes/Engine/Engine.h
@@ -1065,6 +1065,10 @@ public:
 	UPROPERTY(globalconfig)
 	FString TextureColorViewModeMaterialName;
 
+	/** Path of the texture used to indicate zen streaming is active. */
+	UPROPERTY(globalconfig)
+	FSoftObjectPath DefaultZenStreamingTextureName;
+
 #if WITH_EDITORONLY_DATA
 	/** Material used to render bone weights on skeletal meshes */
 	UPROPERTY()
diff --git a/Engine/Source/Runtime/StorageServerClient/Private/ZenStreamingSettings.h b/Engine/Source/Runtime/StorageServerClient/Private/ZenStreamingSettings.h
new file mode 100644
--- /dev/null
+++ b/Engine/Source/Runtime/StorageServerClient/Private/ZenStreamingSettings.h
@@ -0,0 +1,36 @@
+// Copyright Epic Games, Inc. All Rights Reserved.
+
+#pragma once
+
+#include "Engine/DeveloperSettings.h"
+#include "HAL/IConsoleManager.h"
+
+#include "ZenStreamingSettings.generated.h"
+
+UCLASS(config=Game, defaultconfig, meta=(DisplayName="Zen Streaming"))
+class STORAGESERVERCLIENT_API UZenStreamingSettings : public UDeveloperSettings
+{
+	GENERATED_BODY()
+
+public:
+	/**
+	* Show an in-game badge to indicate if the client is streaming from Zen
+	*/
+	UPROPERTY(config, EditAnywhere, Category = ZenStreaming, meta=(DisplayName="Show Zen streaming badge", ConsoleVariable="zen.indicator.show"))
+	bool bShowZenBadge = true;
+
+	UPROPERTY(config, EditAnywhere, Category = ZenStreaming, meta=(ClampMin=0, ClampMax=1.0, ConsoleVariable="zen.indicator.x"))
+	float BadgePositionX = 0.05;
+
+	UPROPERTY(config, EditAnywhere, Category = ZenStreaming, meta=(ClampMin=0, ClampMax=1.0, ConsoleVariable="zen.indicator.y"))
+	float BadgePositionY = 0.8;
+
+	/**
+	 * For how long (in seconds) should the badge stay on screen. < 0 will make the badge persistent
+	 */
+	UPROPERTY(config, EditAnywhere, Category = ZenStreaming, meta=(ClampMin=-1, ConsoleVariable="zen.indicator.fadetime"))
+	float BadgeFadeTime = -1;
+
+	UPROPERTY(config, EditAnywhere, Category = ZenStreaming, meta=(ClampMin=0, ClampMax=1.0, ConsoleVariable="zen.indicator.alpha"))
+	float BadgeAlpha = 0.5;
+};
\ No newline at end of file
diff --git a/Engine/Source/Runtime/StorageServerClient/StorageServerClient.Build.cs b/Engine/Source/Runtime/StorageServerClient/StorageServerClient.Build.cs
--- a/Engine/Source/Runtime/StorageServerClient/StorageServerClient.Build.cs
+++ b/Engine/Source/Runtime/StorageServerClient/StorageServerClient.Build.cs
@@ -10,6 +10,7 @@ public class StorageServerClient : ModuleRules
 			new string[] {
 				"Core",
 				"CoreUObject",
+				"DeveloperSettings",
 				"Sockets",
 				"CookOnTheFly",
 				"Json",
diff --git a/Engine/Source/Runtime/StorageServerClientDebug/Private/StorageServerConnectionDebug.cpp b/Engine/Source/Runtime/StorageServerClientDebug/Private/StorageServerConnectionDebug.cpp
--- a/Engine/Source/Runtime/StorageServerClientDebug/Private/StorageServerConnectionDebug.cpp
+++ b/Engine/Source/Runtime/StorageServerClientDebug/Private/StorageServerConnectionDebug.cpp
@@ -4,6 +4,7 @@
 
 #include "Debug/DebugDrawService.h"
 #include "Engine/GameEngine.h"
+#include "Engine/Texture2D.h"
 #include "Modules/ModuleInterface.h"
 #include "Modules/ModuleManager.h"
 #include "ProfilingDebugging/CsvProfiler.h"
@@ -17,6 +18,7 @@
 #include "HAL/PlatformProcess.h"
 #include "Misc/ScopeLock.h"
 #include "Misc/CoreDelegates.h"
+#include "Misc/ConfigCacheIni.h"
 
 #if !UE_BUILD_SHIPPING
 
@@ -32,6 +34,55 @@ TRACE_DECLARE_UNCHECKED_FLOAT_COUNTER(ZenClient_MaxReqThroughputMbps, TEXT("ZenC
 TRACE_DECLARE_UNCHECKED_FLOAT_COUNTER(ZenClient_MinReqThroughputMbps, TEXT("ZenClient/MinReqThroughputMbps (decompressed)"));
 TRACE_DECLARE_UNCHECKED_INT_COUNTER(  ZenClient_RequestCountPerSec,   TEXT("ZenClient/RequestCountPerSec"));
 
+static bool GZenShowIndicator = true;
+static FAutoConsoleVariableRef CVarShowZenIndicator(
+	TEXT("zen.indicator.show"),
+	GZenShowIndicator,
+	TEXT("Show on-screen indicator when Zen streaming is active"),
+	ECVF_Default
+);
+
+static bool GZenShowDebugMessage = false;
+static FAutoConsoleVariableRef CVarShowZenDebugOnScreenMessage(
+	TEXT("zen.onscreenmessage"),
+	GZenShowDebugMessage,
+	TEXT("Show an on-screen message with zen streaming stats"),
+	ECVF_Default
+);
+
+static float GZenIndicatorPosX = 0.01f;
+static FAutoConsoleVariableRef CVarZenIndicatorPosX(
+	TEXT("zen.indicator.x"),
+	GZenIndicatorPosX,
+	TEXT("Zen on-screen indicator position (horizontal)"),
+	ECVF_Default
+);
+
+static float GZenIndicatorPosY = 0.8f;
+static FAutoConsoleVariableRef CVarZenIndicatorPosY(
+	TEXT("zen.indicator.y"),
+	GZenIndicatorPosY,
+	TEXT("Zen on-screen indicator position (vertical)"),
+	ECVF_Default
+);
+
+static float GZenIndicatorFadeTime = -1.f;
+static FAutoConsoleVariableRef CVarZenIndicatorFadeTime(
+	TEXT("zen.indicator.fadetime"),
+	GZenIndicatorFadeTime,
+	TEXT("Zen on-screen indicator fade time in seconds"),
+	ECVF_Default
+);
+
+static constexpr float IndicatorFadeSpeed = 5.f;
+
+static float GZenIndicatorAlpha = 0.5f;
+static FAutoConsoleVariableRef CVarZenIndicatorAlpha(
+	TEXT("zen.indicator.alpha"),
+	GZenIndicatorAlpha,
+	TEXT("Zen on-screen indicator transparency"),
+	ECVF_Default
+);
 
 static bool GZenShowGraphs = false;
 static FAutoConsoleVariableRef CVarZenShowGraphs(
@@ -54,6 +105,9 @@ namespace
 {
 	static constexpr int    OneMinuteSeconds = 60;
 	static constexpr double WidthSeconds = OneMinuteSeconds * 0.25;
+
+	constexpr float			ZenIconPadding = 8.f;
+	constexpr float			ZenIndicatorTextWidth = 256;
 }
 
 
@@ -65,6 +119,8 @@ bool FStorageServerConnectionDebug::OnTick(float)
 	
 	double StatsTimeNow = FPlatformTime::Seconds();
 	double Duration = StatsTimeNow - UpdateStatsTime;
+	IndicatorElapsedTime += StatsTimeNow - IndicatorLastTime;
+	IndicatorLastTime = StatsTimeNow;
 
 	//Persistent debug message and CSV stats
 	if (Duration > UpdateStatsTimer)
@@ -82,11 +138,11 @@ bool FStorageServerConnectionDebug::OnTick(float)
 			ReqCount = ceil((double)Stats.RequestCount / Duration);
 		}
 
-		if (GZenShowStats && GEngine)
+		if (GZenShowDebugMessage && GEngine)
 		{
 			FString ZenConnectionDebugMsg;
 			ZenConnectionDebugMsg = FString::Printf(TEXT("ZenServer %s from %s [%.2fMbps]"), UE::IsUsingZenPakFileStreaming() ? TEXT("pak streaming") : TEXT("streaming"), * HostAddress, Throughput);
-			GEngine->AddOnScreenDebugMessage((uint64)this, 86400.0f, FColor::White, ZenConnectionDebugMsg, false);
+			GEngine->AddOnScreenDebugMessage((uint64)this, UpdateStatsTimer, FColor::White, ZenConnectionDebugMsg, false);
 		}
 		
 		History.push_back({ StatsTimeNow, MaxReqThroughput, MinReqThroughput, Throughput, ReqCount });
@@ -226,8 +282,160 @@ void FStorageServerConnectionDebug::OnDraw(UCanvas* Canvas, APlayerController*)
 		}
 		HeightScaleRequest = FMath::Min(MaxHeightScaleRequest, FMath::Max(MaxValueInHistory, 1.0));
 	}
+
+	if (GZenShowIndicator)
+	{
+		static bool SettingsLoaded = false;
+		if (SettingsLoaded == false)
+		{
+			LoadZenStreamingSettings();
+			SettingsLoaded = true;
+		}
+		
+		DrawZenIndicator(Canvas);
+	}
+	else if (bDestroyZenIcon)
+	{
+		DestroyZenIcon();
+	}
 }
 
+void FStorageServerConnectionDebug::LoadZenStreamingSettings()
+{
+	float Val = 0.f;
+	bool bShowIndicator = true;
+	bShowIndicator = GConfig->GetBoolOrDefault(TEXT("/Script/StorageServerClient.ZenStreamingSettings"), TEXT("zen.indicator.show"), GZenShowIndicator, GGameIni);
+	if (bShowIndicator != GZenShowIndicator)
+	{
+		CVarShowZenIndicator->Set(bShowIndicator, ECVF_SetByGameSetting);
+	}
+
+	Val = GConfig->GetFloatOrDefault(TEXT("/Script/StorageServerClient.ZenStreamingSettings"), TEXT("zen.indicator.x"), GZenIndicatorPosX, GGameIni);
+	if (Val != GZenIndicatorPosX)
+	{
+		CVarZenIndicatorPosX->Set(Val, ECVF_SetByGameSetting);
+	}
+
+	Val = GConfig->GetFloatOrDefault(TEXT("/Script/StorageServerClient.ZenStreamingSettings"), TEXT("zen.indicator.y"), GZenIndicatorPosY, GGameIni);
+	if (Val != GZenIndicatorPosY)
+	{
+		CVarZenIndicatorPosY->Set(Val, ECVF_SetByGameSetting);
+	}
+	
+	Val = GConfig->GetFloatOrDefault(TEXT("/Script/StorageServerClient.ZenStreamingSettings"), TEXT("zen.indicator.fadetime"), GZenIndicatorFadeTime, GGameIni);
+	if (Val != GZenIndicatorFadeTime)
+	{
+		CVarZenIndicatorFadeTime->Set(Val, ECVF_SetByGameSetting);
+	}
+	
+	Val = GConfig->GetFloatOrDefault(TEXT("/Script/StorageServerClient.ZenStreamingSettings"), TEXT("zen.indicator.alpha"), GZenIndicatorAlpha, GGameIni);
+	if (Val != GZenIndicatorAlpha)
+	{
+		CVarZenIndicatorAlpha->Set(Val, ECVF_SetByGameSetting);
+	}
+
+	CVarShowZenIndicator->SetOnChangedCallback(FConsoleVariableDelegate::CreateLambda([this](IConsoleVariable* Var)
+		{
+			IndicatorElapsedTime = 0.f;
+			if (!Var->GetBool() && ZenIcon != nullptr)
+			{
+				bDestroyZenIcon = true;
+			}
+		}
+	));
+
+	CVarZenIndicatorFadeTime->SetOnChangedCallback(FConsoleVariableDelegate::CreateLambda([this](IConsoleVariable* Var)
+		{
+			IndicatorElapsedTime = 0.f;
+		}
+	));
+}
+
+void FStorageServerConnectionDebug::DrawZenIndicator(UCanvas* Canvas)
+{
+	if (ZenIcon == nullptr)
+	{
+		CreateZenIcon();
+	}
+
+	if (ZenIcon)
+	{
+		FCanvasIcon CanvasIcon = UCanvas::MakeIcon(ZenIcon);
+		UFont* StringFont = GEngine->GetSmallFont();
+
+		float Alpha = FMath::Clamp(GZenIndicatorAlpha, 0.0f, 1.0f);
+		const float DeltaTime = FMath::Clamp(IndicatorElapsedTime - GZenIndicatorFadeTime, 0.f, IndicatorFadeSpeed);
+		const float AlphaDelta = GZenIndicatorAlpha / IndicatorFadeSpeed;
+		if (GZenIndicatorFadeTime > 0 && DeltaTime > 0)
+		{
+			Alpha = GZenIndicatorAlpha - DeltaTime * AlphaDelta;
+		}
+
+		const float IconPosX = (Canvas->ClipX - ZenIcon->GetSurfaceWidth()) * GZenIndicatorPosX;
+		const float IconPosY = (Canvas->ClipY - ZenIcon->GetSurfaceHeight()) * GZenIndicatorPosY;
+		const float BackgroundPosX = IconPosX - ZenIconPadding;
+		const float BackgroundPosY = IconPosY - ZenIconPadding;
+		const float BackgroundSizeX = ZenIcon->GetSurfaceWidth() + 2 * ZenIconPadding + ZenIndicatorTextWidth;
+		const float BackgroundSizeY = ZenIcon->GetSurfaceHeight() + 2 * ZenIconPadding;
+		const float IconScreenPosX = IconPosX / Canvas->ClipX;
+		const float IconScreenPosY = IconPosY / Canvas->ClipY;
+		const float TextPosX = IconPosX + ZenIcon->GetSurfaceWidth() + 2 * ZenIconPadding;
+		const float TextPosY = IconPosY + ZenIconPadding;
+
+		const FLinearColor BackgroundColor(0.0f, 0.0f, 0.0f, Alpha);
+		const FLinearColor TextBackgroundColor(0.7f, 0.7f, 0.7f, Alpha);
+		const FLinearColor DebugBackgroundColor(1.f, 0.f, 0.f, Alpha);
+
+		FCanvasTileItem BackgroundTileItem(
+			FVector2D(BackgroundPosX, BackgroundPosY),
+			FVector2D(BackgroundSizeX, BackgroundSizeY),
+			BackgroundColor);
+
+		FCanvasTileItem DebugBackgroundTile(
+			FVector2D(IconPosX, IconPosY),
+			FVector2D(ZenIcon->GetSurfaceWidth(), ZenIcon->GetSurfaceHeight()),
+			DebugBackgroundColor);
+
+		DebugBackgroundTile.BlendMode = SE_BLEND_AlphaBlend;
+		BackgroundTileItem.BlendMode = SE_BLEND_AlphaBlend;
+
+		Canvas->DrawItem(BackgroundTileItem);
+		Canvas->SetDrawColor(255, 255, 255, (uint8)(Alpha * 255.f));
+		Canvas->DrawIcon(CanvasIcon, IconPosX, IconPosY, 1.f);
+		FString ZenStreamingString = FString::Printf(TEXT("Zen %s from: %s"), UE::IsUsingZenPakFileStreaming() ? TEXT("pak streaming") : TEXT("streaming"), * HostAddress);
+		Canvas->Canvas->DrawShadowedString(TextPosX, TextPosY, *ZenStreamingString, StringFont, TextBackgroundColor);
+		FString ZenBandwidthString = FString::Printf(TEXT("Bandwidth: %.2f Mbps"), Throughput);
+		Canvas->Canvas->DrawShadowedString(TextPosX, TextPosY + StringFont->GetMaxCharHeight() + 8, *ZenBandwidthString, StringFont, TextBackgroundColor);
+
+		if (Alpha == 0.f)
+		{
+			bDestroyZenIcon = true;
+			GZenShowIndicator = false;
+		}
+	}
+}
+
+void FStorageServerConnectionDebug::CreateZenIcon()
+{
+	if (ZenIcon == nullptr && GEngine->DefaultZenStreamingTextureName.IsValid())
+	{
+		ZenIcon = Cast<UTexture2D>(GEngine->DefaultZenStreamingTextureName.TryLoad());
+		if (FPlatformProperties::RequiresCookedData() && ZenIcon != nullptr)
+		{
+			ZenIcon->AddToRoot();
+		}
+	}
+}
+
+void FStorageServerConnectionDebug::DestroyZenIcon()
+{
+	if (ZenIcon != nullptr && FPlatformProperties::RequiresCookedData())
+	{
+		ZenIcon->RemoveFromRoot();
+		ZenIcon = nullptr;
+		bDestroyZenIcon = false;
+	}
+}
 
 class FStorageServerClientDebugModule
 	: public IModuleInterface
diff --git a/Engine/Source/Runtime/StorageServerClientDebug/Private/StorageServerConnectionDebug.h b/Engine/Source/Runtime/StorageServerClientDebug/Private/StorageServerConnectionDebug.h
--- a/Engine/Source/Runtime/StorageServerClientDebug/Private/StorageServerConnectionDebug.h
+++ b/Engine/Source/Runtime/StorageServerClientDebug/Private/StorageServerConnectionDebug.h
@@ -16,12 +16,19 @@ public:
 		: StorageServerPlatformFile(InStorageServerPlatformFile)
 		, HostAddress(InStorageServerPlatformFile->GetHostAddr())
 	{
+		IndicatorLastTime = FPlatformTime::Seconds();
 	}
 
 	bool OnTick(float); // FTickerDelegate
 	void OnDraw(UCanvas*, APlayerController*); // FDebugDrawDelegate
 
 private:
+	void LoadZenStreamingSettings();
+
+	void DrawZenIndicator(UCanvas* Canvas);
+	void CreateZenIcon();
+	void DestroyZenIcon();
+
 	double MaxReqThroughput = 0.0;
 	double MinReqThroughput = 0.0;
 	uint32 ReqCount = 0;
@@ -45,5 +52,10 @@ private:
 	FString HostAddress;
 
 	FCriticalSection CS;
+
+	class UTexture2D* ZenIcon = nullptr;
+	bool bDestroyZenIcon = false;
+	double IndicatorElapsedTime = 0.0;
+	double IndicatorLastTime = 0.0;
 };
 #endif // !UE_BUILD_SHIPPING
