Remove the requirement for cooked data for the Zen streaming icon

#rnx

diff --git a/Engine/Source/Runtime/StorageServerClientDebug/Private/StorageServerConnectionDebug.cpp b/Engine/Source/Runtime/StorageServerClientDebug/Private/StorageServerConnectionDebug.cpp
--- a/Engine/Source/Runtime/StorageServerClientDebug/Private/StorageServerConnectionDebug.cpp
+++ b/Engine/Source/Runtime/StorageServerClientDebug/Private/StorageServerConnectionDebug.cpp
@@ -420,7 +420,7 @@ void FStorageServerConnectionDebug::CreateZenIcon()
 	if (ZenIcon == nullptr && GEngine->DefaultZenStreamingTextureName.IsValid())
 	{
 		ZenIcon = Cast<UTexture2D>(GEngine->DefaultZenStreamingTextureName.TryLoad());
-		if (FPlatformProperties::RequiresCookedData() && ZenIcon != nullptr)
+		if (ZenIcon != nullptr)
 		{
 			ZenIcon->AddToRoot();
 		}
@@ -429,7 +429,7 @@ void FStorageServerConnectionDebug::CreateZenIcon()
 
 void FStorageServerConnectionDebug::DestroyZenIcon()
 {
-	if (ZenIcon != nullptr && FPlatformProperties::RequiresCookedData())
+	if (ZenIcon != nullptr)
 	{
 		ZenIcon->RemoveFromRoot();
 		ZenIcon = nullptr;
